Game Engine: Auxiliary camera viewport client

Additional viewports can be registered to the engine, so that custom viewport clients (and this new Camera Viewport Client) can participate in the engine draw loop.

RedrawViewports is moved up from Game Engine to base Engine, iterates all registered viewports, called by both Game and Editor engines so your code will work the same in both scenarios.  PIE viewport's can't (yet) be merged into the same loop because of per-world PIE-specific setup/teardown.  Other editor viewports aren't because they have special throttling that doesn't make sense to introduce to Game Engine (customers who want throttling for a more productivity-like application can implement their own Engine subclass).

Camera-cut flags are now cleared from the Game Viewport Client instead of from all viewports (a more appropriate place for this).

UpdateLevelStreaming removed from FViewport::Draw
- AND UGameViewportClient::Draw (it was called several times per frame from different places)
- Consolidated into `UGameEngine::Tick` after RedrawViewports, matching what !GUseUnifiedTimeBudgetForStreaming already did; editor engine already calls it
- AddStreamingViewInfo extracted as a shared helper to avoid copy-paste among 3D view clients

[RN] Support Drawing of more than one FViewport per frame in UGameEngine, introduce experimental plugin with single-camera viewport client

#Synced-CL 52671779

diff --git a/Engine/Plugins/Runtime/ModularViewports/ModularViewports.uplugin b/Engine/Plugins/Runtime/ModularViewports/ModularViewports.uplugin
new file mode 100644
--- /dev/null
+++ b/Engine/Plugins/Runtime/ModularViewports/ModularViewports.uplugin
@@ -0,0 +1,22 @@
+{
+	"FileVersion": 3,
+	"Version": 1,
+	"VersionName": "0.1",
+	"FriendlyName": "Modular Viewports",
+	"Description": "Viewport Client implementations that provide more flexibility and granularity than Game Viewport Client.",
+	"Category": "Rendering",
+	"CreatedBy": "Epic Games, Inc.",
+	"CreatedByURL": "https://epicgames.com",
+	"EnabledByDefault": false,
+	"CanContainContent": false,
+	"IsExperimentalVersion": true,
+	"Installed": false,
+	"Modules":
+	[
+		{
+			"Name": "ModularViewports",
+			"Type": "Runtime",
+			"LoadingPhase": "Default"
+		}
+	]
+}
diff --git a/Engine/Plugins/Runtime/ModularViewports/Source/ModularViewports/ModularViewports.Build.cs b/Engine/Plugins/Runtime/ModularViewports/Source/ModularViewports/ModularViewports.Build.cs
new file mode 100644
--- /dev/null
+++ b/Engine/Plugins/Runtime/ModularViewports/Source/ModularViewports/ModularViewports.Build.cs
@@ -0,0 +1,24 @@
+// Copyright Epic Games, Inc. All Rights Reserved.
+
+namespace UnrealBuildTool.Rules
+{
+	public class ModularViewports : ModuleRules
+	{
+		public ModularViewports(ReadOnlyTargetRules Target) : base(Target)
+		{
+			PublicDependencyModuleNames.AddRange([
+				"Core",
+			]);
+
+			PrivateDependencyModuleNames.AddRange([
+				"CoreUObject",
+				"Engine",
+				"RenderCore",
+				"Renderer",
+				"RHI",
+				"Slate",
+				"SlateCore",
+			]);
+		}
+	}
+}
diff --git a/Engine/Plugins/Runtime/ModularViewports/Source/ModularViewports/Private/Camera/CameraViewportClient.cpp b/Engine/Plugins/Runtime/ModularViewports/Source/ModularViewports/Private/Camera/CameraViewportClient.cpp
new file mode 100644
--- /dev/null
+++ b/Engine/Plugins/Runtime/ModularViewports/Source/ModularViewports/Private/Camera/CameraViewportClient.cpp
@@ -0,0 +1,210 @@
+// Copyright Epic Games, Inc. All Rights Reserved.
+
+#include "Camera/CameraViewportClient.h"
+#include "CanvasTypes.h"
+#include "EngineModule.h"
+#include "LegacyScreenPercentageDriver.h"
+#include "UnrealClient.h"
+#include "Framework/Application/SlateApplication.h"
+#include "SceneView.h"
+#include "Engine/Engine.h"
+#include "Engine/LocalPlayer.h"
+#include "Engine/World.h"
+#include "GameFramework/WorldSettings.h"
+#include "Camera/CameraComponent.h"
+#include "SceneManagement.h"
+
+namespace UE::CameraViewportClient
+{
+	namespace
+	{
+		void CalculateProjectionData(const FViewport& Viewport, const UCameraComponent& Camera, FSceneViewProjectionData& OutProjectionData)
+		{
+			const FTransform CameraTransform = Camera.GetComponentTransform();
+			OutProjectionData.ViewOrigin = CameraTransform.GetLocation();
+
+			FIntPoint ViewportSize = Viewport.GetSizeXY();
+			OutProjectionData.SetViewRectangle(FIntRect(0, 0, ViewportSize.X, ViewportSize.Y));
+
+			OutProjectionData.ViewRotationMatrix = FInverseRotationMatrix(CameraTransform.Rotator()) * FMatrix(
+				FPlane(0, 0, 1, 0),
+				FPlane(1, 0, 0, 0),
+				FPlane(0, 1, 0, 0),
+				FPlane(0, 0, 0, 1));
+			const float MinZ = GNearClippingPlane;
+			const float MaxZ = MinZ;
+
+			float XAxisMultiplier;
+			float YAxisMultiplier;
+			const EAspectRatioAxisConstraint FovType = Camera.bOverrideAspectRatioAxisConstraint
+				? Camera.AspectRatioAxisConstraint
+				: GetDefault<ULocalPlayer>()->AspectRatioAxisConstraint;
+			const bool bMaintainXFOV = (ViewportSize.X > ViewportSize.Y && FovType == AspectRatio_MajorAxisFOV) || FovType == AspectRatio_MaintainXFOV;
+			const float ViewportAspectRatio = FMath::Max(1.0f, static_cast<float>(ViewportSize.X)) / FMath::Max(1.0f, static_cast<float>(ViewportSize.Y));
+			float MatrixFov;
+			if (bMaintainXFOV)
+			{
+				XAxisMultiplier = 1.0f;
+				YAxisMultiplier = ViewportAspectRatio;
+				MatrixFov = FMath::Max(0.001f, Camera.FieldOfView) * UE_PI / 360.0f;
+			}
+			else
+			{
+				XAxisMultiplier = 1 / ViewportAspectRatio;
+				YAxisMultiplier = 1.0f;
+				const float HalfXFov = FMath::DegreesToRadians(FMath::Max(0.001f, Camera.FieldOfView) / 2.f);
+				const float HalfYFov = FMath::Atan(FMath::Tan(HalfXFov) / FMath::Max(0.001f, Camera.AspectRatio));
+				MatrixFov = HalfYFov;
+			}
+			OutProjectionData.ProjectionMatrix = FReversedZPerspectiveMatrix(MatrixFov, MatrixFov, XAxisMultiplier, YAxisMultiplier, MinZ, MaxZ);
+		}
+
+		TUniquePtr<FSceneViewFamilyContext> BuildViewFamily(const FViewport& Viewport, FSceneViewStateReference& ViewState, const UWorld& World, const UCameraComponent& Camera)
+		{
+			TUniquePtr<FSceneViewFamilyContext> ViewFamily = MakeUnique<FSceneViewFamilyContext>(
+				FSceneViewFamily::ConstructionValues(
+					&Viewport,
+					World.Scene,
+					ESFIM_Game
+				)
+				.SetRealtimeUpdate(true)
+				.SetRequireMobileMultiView(false)
+			);
+			EngineShowFlagOverride(ESFIM_Game, ViewFamily->ViewMode, ViewFamily->EngineShowFlags, false);
+
+			FSceneViewInitOptions ViewInitOptions;
+
+			const AWorldSettings* const WorldSettings = World.GetWorldSettings();
+			if (WorldSettings)
+			{
+				ViewInitOptions.WorldToMetersScale = WorldSettings->WorldToMeters;
+			}
+
+			CalculateProjectionData(Viewport, Camera, ViewInitOptions);
+
+			ViewInitOptions.ViewFamily = ViewFamily.Get();
+			ViewInitOptions.SceneViewStateInterface = ViewState.GetReference();
+
+			ViewInitOptions.BackgroundColor = FLinearColor::Black;
+			ViewInitOptions.FOV = Camera.FieldOfView;
+
+			const FTransform CameraTransform = Camera.GetComponentTransform();
+			FSceneView* View = new FSceneView(ViewInitOptions);
+			View->ViewLocation = CameraTransform.GetLocation();
+			View->ViewRotation = CameraTransform.Rotator();
+			ViewFamily->Views.Add(View);
+
+			View->StartFinalPostprocessSettings(View->ViewLocation);
+			View->OverridePostProcessSettings(Camera.PostProcessSettings, Camera.PostProcessBlendWeight);
+			View->EndFinalPostprocessSettings(ViewInitOptions);
+
+			ViewFamily->bIsHDR = Viewport.IsHDRViewport();
+
+			return ViewFamily;
+		}
+
+		void RenderViewFamily(TUniquePtr<FSceneViewFamilyContext>&& ViewFamily, FCanvas& Canvas)
+		{
+			if (ViewFamily->GetScreenPercentageInterface() == nullptr)
+			{
+				ViewFamily->SetScreenPercentageInterface(new FLegacyScreenPercentageDriver(*ViewFamily, 1.0));
+			}
+			GetRendererModule().BeginRenderingViewFamily(&Canvas, ViewFamily.Get());
+
+			if (FSlateApplication::Get().GetPlatformApplication()->IsAllowedToRender())
+			{
+				Canvas.Flush_GameThread();
+			}
+		}
+	}
+}
+
+namespace UE
+{
+	FCameraViewportClient::FCameraViewportClient(TWeakObjectPtr<UCameraComponent> InCameraComponent)
+		: Camera(InCameraComponent)
+	{
+	}
+
+	UWorld* FCameraViewportClient::GetWorld() const
+	{
+		return Camera.IsValid()
+			? Camera->GetWorld()
+			: nullptr;
+	}
+
+	void FCameraViewportClient::AddReferencedObjects(FReferenceCollector& Collector)
+	{
+		FSceneViewStateInterface* Ref = ViewState.GetReference();
+		if (Ref)
+		{
+			Ref->AddReferencedObjects(Collector);
+		}
+	}
+
+	FString FCameraViewportClient::GetReferencerName() const
+	{
+		return TEXT("FCameraViewportClient");
+	}
+
+	void FCameraViewportClient::CalculateProjectionData(const FViewport& Viewport, FSceneViewProjectionData& OutProjectionData) const
+	{
+		if (!Camera.IsValid())
+		{
+			return;
+		}
+
+		CameraViewportClient::CalculateProjectionData(Viewport, *Camera, OutProjectionData);
+	}
+
+	void FCameraViewportClient::Draw(FViewport* Viewport, FCanvas* Canvas)
+	{
+		using namespace CameraViewportClient;
+
+		check(Viewport);
+		check(Canvas);
+		UWorld* World = GetWorld();
+		if (!World)
+		{
+			return;
+		}
+
+		// Hypothetically, a camera can move from one world to another
+		if (ViewStateWorld != World)
+		{
+			ViewState.Destroy();
+			ViewState.Allocate(World->GetFeatureLevel());
+			ViewStateWorld = World;
+		}
+
+		TUniquePtr<FSceneViewFamilyContext> ViewFamily = BuildViewFamily(*Viewport, ViewState, *World, *Camera);
+		AddStreamingViewInfo(*World, *ViewFamily->Views[0]);
+		RenderViewFamily(MoveTemp(ViewFamily), *Canvas);
+	}
+
+	bool FCameraViewportClient::InputChar(FViewport* const Viewport, const int32 ControllerId, const TCHAR Character)
+	{
+		// Absorb all keys so game input events are not routed to the Slate editor frame
+		return Viewport->IsSlateViewport() && Viewport->IsPlayInEditorViewport();
+	}
+
+	bool FCameraViewportClient::IsFocused(FViewport* Viewport)
+	{
+		return Viewport->HasFocus() || Viewport->HasMouseCapture();
+	}
+
+	bool FCameraViewportClient::RequiresHitProxyStorage()
+	{
+		return false;
+	}
+
+	bool FCameraViewportClient::CaptureMouseOnLaunch()
+	{
+		return false;
+	}
+
+	bool FCameraViewportClient::ShouldDPIScaleSceneCanvas() const
+	{
+		return false;
+	}
+}
diff --git a/Engine/Plugins/Runtime/ModularViewports/Source/ModularViewports/Private/ModularViewportsModule.cpp b/Engine/Plugins/Runtime/ModularViewports/Source/ModularViewports/Private/ModularViewportsModule.cpp
new file mode 100644
--- /dev/null
+++ b/Engine/Plugins/Runtime/ModularViewports/Source/ModularViewports/Private/ModularViewportsModule.cpp
@@ -0,0 +1,5 @@
+// Copyright Epic Games, Inc. All Rights Reserved.
+
+#include "Modules/ModuleManager.h"
+
+IMPLEMENT_MODULE(FDefaultModuleImpl, ModularViewports);
diff --git a/Engine/Plugins/Runtime/ModularViewports/Source/ModularViewports/Public/Camera/CameraViewportClient.h b/Engine/Plugins/Runtime/ModularViewports/Source/ModularViewports/Public/Camera/CameraViewportClient.h
new file mode 100644
--- /dev/null
+++ b/Engine/Plugins/Runtime/ModularViewports/Source/ModularViewports/Public/Camera/CameraViewportClient.h
@@ -0,0 +1,56 @@
+// Copyright Epic Games, Inc. All Rights Reserved.
+
+#pragma once
+
+#include "CoreMinimal.h"
+#include "ViewportClient.h"
+#include "SceneView.h"
+#include "UObject/GCObject.h"
+#include "SceneTypes.h"
+
+#define UE_API MODULARVIEWPORTS_API
+
+class UCameraComponent;
+class FSceneViewFamilyContext;
+
+namespace UE
+{
+	/** A Viewport Client that renders a single UCameraComponent. */
+	class FCameraViewportClient : public FCommonViewportClient, public FGCObject
+	{
+		TWeakObjectPtr<UCameraComponent> Camera;
+		TWeakObjectPtr<UWorld> ViewStateWorld;
+		FSceneViewStateReference ViewState;
+
+	public:
+		UE_API FCameraViewportClient(TWeakObjectPtr<UCameraComponent>);
+
+		const TWeakObjectPtr<UCameraComponent>& GetCameraComponent() const
+		{
+			return Camera;
+		}
+
+		/** Extracts projection data (view rect, rotation matrix, projection matrix) from the camera.
+		 * 
+		 * Only safe to call within GC Lock.
+		 */
+		UE_API void CalculateProjectionData(const FViewport& Viewport, FSceneViewProjectionData& OutProjectionData) const;
+
+		// FGCObject interface
+		virtual void AddReferencedObjects(FReferenceCollector& Collector) override;
+		virtual FString GetReferencerName() const override;
+		// ~FGCObject interface
+
+		// FViewportClient interface
+		virtual UWorld* GetWorld() const override;
+		virtual void Draw(FViewport* Viewport, FCanvas* Canvas) override;
+		virtual bool InputChar(FViewport* const Viewport, const int32 ControllerId, const TCHAR Character) override;
+		virtual bool IsFocused(FViewport* Viewport) override;
+		virtual bool RequiresHitProxyStorage() override;
+		virtual bool CaptureMouseOnLaunch() override;
+		virtual bool ShouldDPIScaleSceneCanvas() const override;
+		// ~FViewportClient interface
+	};
+}
+
+#undef UE_API
diff --git a/Engine/Plugins/VirtualProduction/Avalanche/Source/Avalanche/Private/Framework/AvaGameViewportClient.cpp b/Engine/Plugins/VirtualProduction/Avalanche/Source/Avalanche/Private/Framework/AvaGameViewportClient.cpp
--- a/Engine/Plugins/VirtualProduction/Avalanche/Source/Avalanche/Private/Framework/AvaGameViewportClient.cpp
+++ b/Engine/Plugins/VirtualProduction/Avalanche/Source/Avalanche/Private/Framework/AvaGameViewportClient.cpp
@@ -172,6 +172,7 @@ void UAvaGameViewportClient::Draw(FViewport* InViewport, FCanvas* InCanvas)
 	}
 
 	TArray<FSceneView*> Views;
+	bool bAnyPlayerCameraCut = false;
 
 	for (FLocalPlayerIterator PlayerIterator(GEngine, World); PlayerIterator; ++PlayerIterator)
 	{
@@ -229,11 +230,19 @@ void UAvaGameViewportClient::Draw(FViewport* InViewport, FCanvas* InCanvas)
 				World->LastRenderTime = World->GetTimeSeconds();
 			}
 		}
-	}
 
-	// Update level streaming.
-	World->UpdateLevelStreaming();
+		if (const APlayerController* const PlayerController = LocalPlayer->PlayerController)
+		{
+			if (APlayerCameraManager* const CameraManager = PlayerController->PlayerCameraManager)
+			{
+				bAnyPlayerCameraCut = bAnyPlayerCameraCut || CameraManager->bGameCameraCutThisFrame;
+				// Assumes any other systems that call ULocalPlayer::CalcSceneView(...) and care about camera cuts will do so before this happens.
+				CameraManager->bGameCameraCutThisFrame = false;
+			}
+		}
+	}
 
+	InViewport->SetCameraCut(bAnyPlayerCameraCut);
 	InCanvas->Clear(FLinearColor::Transparent);
 
 	// If a screen percentage interface was not set by one of the view extension, then set the legacy one.
diff --git a/Engine/Source/Editor/UnrealEd/Classes/Editor/EditorEngine.h b/Engine/Source/Editor/UnrealEd/Classes/Editor/EditorEngine.h
--- a/Engine/Source/Editor/UnrealEd/Classes/Editor/EditorEngine.h
+++ b/Engine/Source/Editor/UnrealEd/Classes/Editor/EditorEngine.h
@@ -882,6 +882,7 @@ public:
 	UNREALED_API virtual void PreExit() override;
 	UNREALED_API virtual float GetMaxTickRate(float DeltaTime, bool bAllowFrameRateSmoothing = true) const override;
 	UNREALED_API virtual void Tick(float DeltaSeconds, bool bIdleMode) override;
+	UNREALED_API virtual void RedrawViewports(bool bShouldPresent = true) override;
 	UNREALED_API virtual bool ShouldDrawBrushWireframe(AActor* InActor) override;
 	UNREALED_API virtual void NotifyToolsOfObjectReplacement(const TMap<UObject*, UObject*>& OldToNewInstanceMap) override;
 	UNREALED_API virtual bool ShouldThrottleCPUUsage() const override;
diff --git a/Engine/Source/Editor/UnrealEd/Private/EditorEngine.cpp b/Engine/Source/Editor/UnrealEd/Private/EditorEngine.cpp
--- a/Engine/Source/Editor/UnrealEd/Private/EditorEngine.cpp
+++ b/Engine/Source/Editor/UnrealEd/Private/EditorEngine.cpp
@@ -2394,32 +2394,9 @@ void UEditorEngine::Tick( float DeltaSeconds, bool bIdleMode )
 
 	if( FSlateThrottleManager::Get().IsAllowingExpensiveTasks() )
 	{
-		for (auto ContextIt = WorldList.CreateIterator(); ContextIt; ++ContextIt)
-		{
-			FWorldContext &PieContext = *ContextIt;
-			if (PieContext.WorldType != EWorldType::PIE)
-			{
-				continue;
-			}
-
-			PlayWorld = PieContext.World();
-			GameViewport = PieContext.GameViewport;
-
-			// Render playworld. This needs to happen after the other viewports for screenshots to work correctly in PIE.
-			if (PlayWorld && GameViewport && !bIsSimulatingInEditor)
-			{
-				// Use the PlayWorld as the GWorld, because who knows what will happen in the Tick.
-				UWorld* OldGWorld = SetPlayInEditorWorld( PlayWorld );
-
-				// Render everything.
-				GameViewport->LayoutPlayers();
-				check(GameViewport->Viewport);
-				GameViewport->Viewport->Draw();
-
-				// Pop the world
-				RestoreEditorWorld( OldGWorld );
-			}
-		}
+		SCOPE_CYCLE_COUNTER(STAT_RedrawViewports);
+		CSV_SCOPED_TIMING_STAT_EXCLUSIVE(ViewportMisc);
+		RedrawViewports();
 	}
 
 	// Update resource streaming after both regular Editor viewports and PIE had a chance to add viewers.
@@ -2487,6 +2464,38 @@ void UEditorEngine::Tick( float DeltaSeconds, bool bIdleMode )
 	}
 }
 
+void UEditorEngine::RedrawViewports(bool bShouldPresent)
+{
+	for (auto ContextIt = WorldList.CreateIterator(); ContextIt; ++ContextIt)
+	{
+		FWorldContext &PieContext = *ContextIt;
+		if (PieContext.WorldType != EWorldType::PIE)
+		{
+			continue;
+		}
+
+		PlayWorld = PieContext.World();
+		GameViewport = PieContext.GameViewport;
+
+		// Render playworld. This needs to happen after the other viewports for screenshots to work correctly in PIE.
+		if (PlayWorld && GameViewport && !bIsSimulatingInEditor)
+		{
+			// Use the PlayWorld as the GWorld, because who knows what will happen in the Tick.
+			UWorld* OldGWorld = SetPlayInEditorWorld(PlayWorld);
+
+			// Render everything.
+			GameViewport->LayoutPlayers();
+			check(GameViewport->Viewport);
+			GameViewport->Viewport->Draw(bShouldPresent);
+
+			// Pop the world
+			RestoreEditorWorld(OldGWorld);
+		}
+	}
+
+	Super::RedrawViewports(bShouldPresent);
+}
+
 float UEditorEngine::GetMaxTickRate( float DeltaTime, bool bAllowFrameRateSmoothing ) const
 {
 	float MaxTickRate = 0.0f;
diff --git a/Engine/Source/Runtime/Core/Public/Templates/PointerVariants.h b/Engine/Source/Runtime/Core/Public/Templates/PointerVariants.h
--- a/Engine/Source/Runtime/Core/Public/Templates/PointerVariants.h
+++ b/Engine/Source/Runtime/Core/Public/Templates/PointerVariants.h
@@ -147,6 +147,12 @@ namespace UE::Core::Private
 				}, PtrVariant);
 		}
 
+		/** Returns true if this variant holds a UObject pointer (as opposed to a TSharedPtr). */
+		bool IsUObject() const
+		{
+			return PtrVariant.template IsType<FObjectPtrWrapper>();
+		}
+
 	private:
 		template <typename LhsType, bool bLhsStrength, typename RhsType, bool bRhsStrength
 			UE_REQUIRES_FRIEND(UE_REQUIRES_EXPR((LhsType*)nullptr == (RhsType*)nullptr))
diff --git a/Engine/Source/Runtime/Engine/Classes/Engine/Engine.h b/Engine/Source/Runtime/Engine/Classes/Engine/Engine.h
--- a/Engine/Source/Runtime/Engine/Classes/Engine/Engine.h
+++ b/Engine/Source/Runtime/Engine/Classes/Engine/Engine.h
@@ -3486,7 +3486,17 @@ public:
 
 	ENGINE_API virtual bool LoadMap( FWorldContext& WorldContext, FURL URL, class UPendingNetGame* Pending, FString& Error );
 
-	virtual void RedrawViewports( bool bShouldPresent = true ) { }
+	/** Register a viewport to be drawn each frame. */
+	ENGINE_API void RegisterViewport(TSharedRef<FViewport> Viewport);
+
+	/** Unregister a viewport to stop drawing to it. */
+	ENGINE_API void UnregisterViewport(TSharedRef<FViewport> Viewport);
+
+	/**
+	 * Redraws all viewports.
+	 * @param	bShouldPresent	Whether we want this frame to be presented
+	 */
+	ENGINE_API virtual void RedrawViewports(bool bShouldPresent = true);
 
 	virtual void TriggerStreamingDataRebuild() { }
 
@@ -4064,6 +4074,11 @@ private:
 public:
 	/** Set priority and affinity on game thread either from ini file or from FPlatformAffinity::GetGameThreadPriority()*/
 	ENGINE_API void SetPriorityAndAffinityOnGameThread();
+
+protected:
+	/** Viewports to draw each frame. */
+	TArray<TSharedRef<FViewport>> ActiveViewports;
+	UE::FMutex ActiveViewportsMutex;
 };
 
 /** Global engine pointer. Can be 0 so don't use without checking. */
diff --git a/Engine/Source/Runtime/Engine/Private/GameEngine.cpp b/Engine/Source/Runtime/Engine/Private/GameEngine.cpp
--- a/Engine/Source/Runtime/Engine/Private/GameEngine.cpp
+++ b/Engine/Source/Runtime/Engine/Private/GameEngine.cpp
@@ -280,6 +280,7 @@ void UGameEngine::CreateGameViewport( UGameViewportClient* GameViewportClient )
 
 	SceneViewport = MakeShareable( GameViewportClient->CreateGameViewport(GameViewportWidgetRef) );
 	GameViewportClient->Viewport = SceneViewport.Get();
+	RegisterViewport(SceneViewport.ToSharedRef());
 	//GameViewportClient->CreateHighresScreenshotCaptureRegionWidget(); //  Disabled until mouse based input system can be made to work correctly.
 
 	// The viewport widget needs an interface so it knows what should render
@@ -822,21 +823,11 @@ void UGameEngine::OnGameWindowMoved(const TSharedRef<SWindow>& WindowBeingMoved)
 
 void UGameEngine::RedrawViewports( bool bShouldPresent /*= true*/ )
 {
-	SCOPE_CYCLE_COUNTER(STAT_RedrawViewports);
-	CSV_SCOPED_TIMING_STAT_EXCLUSIVE(ViewportMisc);
 	if ( GameViewport != NULL )
 	{
 		GameViewport->LayoutPlayers();
-		if ( GameViewport->Viewport != NULL )
-		{
-			// Make sure we don't attempt to render a frame in parallel with the loading screen, which might deadlock because of interleaving two different frames.
-			// It could happen because the Game thread invokes RedrawViewports right after blocking loads when the loading thread is still spinning.
-			FSlateRenderer* SlateRenderer = FSlateApplication::Get().GetRenderer();
-			FScopeLock ScopeLock(SlateRenderer->GetResourceCriticalSection());
-
-			GameViewport->Viewport->Draw(bShouldPresent);
-		}
 	}
+	Super::RedrawViewports(bShouldPresent);
 }
 
 void UGameEngine::OnViewportResized(FViewport* Viewport, uint32 Unused)
@@ -2056,6 +2047,8 @@ void UGameEngine::Tick( float DeltaSeconds, bool bIdleMode )
 
 	if (bIsRenderingScene && !bRenderingSuspended)
 	{
+		SCOPE_CYCLE_COUNTER(STAT_RedrawViewports);
+		CSV_SCOPED_TIMING_STAT_EXCLUSIVE(ViewportMisc);
 		RedrawViewports();
 	}
 
@@ -2092,18 +2085,15 @@ void UGameEngine::Tick( float DeltaSeconds, bool bIdleMode )
 				GetRendererModule().PostRenderAllViewports();
 			}
 		}
-		else if (!GUseUnifiedTimeBudgetForStreaming)
+
+		if (!GUseUnifiedTimeBudgetForStreaming && !GIsServer && WorldList.Num() > 0)
 		{
-			// Still need to call UpdateLevelStreaming() even when not rendering
-			// unless we are using the unified streaming budget
-			if (GameViewport && GameViewport->Viewport)
+			SCOPE_CYCLE_COUNTER(STAT_UpdateLevelStreaming);
+			for (FWorldContext& Context : WorldList)
 			{
-				if (FViewportClient* ViewportClient = GameViewport->Viewport->GetClient())
+				if (Context.World())
 				{
-					if (UWorld* World = ViewportClient->GetWorld())
-					{
-						World->UpdateLevelStreaming();
-					}
+					Context.World()->UpdateLevelStreaming();
 				}
 			}
 		}
diff --git a/Engine/Source/Runtime/Engine/Private/GameViewportClient.cpp b/Engine/Source/Runtime/Engine/Private/GameViewportClient.cpp
--- a/Engine/Source/Runtime/Engine/Private/GameViewportClient.cpp
+++ b/Engine/Source/Runtime/Engine/Private/GameViewportClient.cpp
@@ -1639,6 +1639,7 @@ void UGameViewportClient::Draw(FViewport* InViewport, FCanvas* SceneCanvas)
 
 	TMap<ULocalPlayer*,FSceneView*> PlayerViewMap;
 	TArray<FSceneView*> Views;
+	bool bAnyPlayerCameraCut = false;
 
 	for (FLocalPlayerIterator Iterator(GEngine, MyWorld); Iterator; ++Iterator)
 	{
@@ -1784,21 +1785,25 @@ void UGameViewportClient::Draw(FViewport* InViewport, FCanvas* SceneCanvas)
 						}
 					}
 
-					// Add view information for resource streaming. Allow up to 5X boost for small FOV.
-					const float StreamingScale = 1.f / FMath::Clamp<float>(View->LODDistanceFactor, .2f, 1.f);
-					IStreamingManager::Get().AddViewInformation(View->ViewMatrices.GetViewOrigin(), View->UnscaledViewRect.Width(), View->UnscaledViewRect.Width() * View->ViewMatrices.GetViewToClip().M[0][0], StreamingScale, false, 0, nullptr, World);
-					MyWorld->ViewLocationsRenderedLastFrame.Add(View->ViewMatrices.GetViewOrigin());
-					FWorldCachedViewInfo& WorldViewInfo = World->CachedViewInfoRenderedLastFrame.AddDefaulted_GetRef();
-					WorldViewInfo.ViewMatrix = View->ViewMatrices.GetWorldToView();
-					WorldViewInfo.ProjectionMatrix = View->ViewMatrices.GetViewToClip();
-					WorldViewInfo.ViewProjectionMatrix = View->ViewMatrices.GetWorldToClip();
-					WorldViewInfo.ViewToWorld = View->ViewMatrices.GetViewToWorld();
+					AddStreamingViewInfo(*World, *View);
 					World->LastRenderTime = World->GetTimeSeconds();
 				}
 			}
+
+			if (PlayerController)
+			{
+				if (APlayerCameraManager* const CameraManager = PlayerController->PlayerCameraManager)
+				{
+					bAnyPlayerCameraCut = bAnyPlayerCameraCut || CameraManager->bGameCameraCutThisFrame;
+					// Assumes any other systems that call ULocalPlayer::CalcSceneView(...) and care about camera cuts will do so before this happens.
+					CameraManager->bGameCameraCutThisFrame = false;
+				}
+			}
 		}
 	}
 
+	InViewport->SetCameraCut(bAnyPlayerCameraCut);
+
 #if CSV_PROFILER_STATS
 	UpdateCsvCameraStats(PlayerViewMap);
 #endif
@@ -1814,12 +1819,6 @@ void UGameViewportClient::Draw(FViewport* InViewport, FCanvas* SceneCanvas)
 
 	FinalizeViews(&ViewFamily, PlayerViewMap);
 
-	if (!GUseUnifiedTimeBudgetForStreaming)
-	{
-		// Update level streaming.
-		MyWorld->UpdateLevelStreaming();
-	}
-
 	// Find largest rectangle bounded by all rendered views.
 	uint32 MinX=InViewport->GetSizeXY().X, MinY=InViewport->GetSizeXY().Y, MaxX=0, MaxY=0;
 	uint32 TotalArea = 0;
diff --git a/Engine/Source/Runtime/Engine/Private/UnrealClient.cpp b/Engine/Source/Runtime/Engine/Private/UnrealClient.cpp
--- a/Engine/Source/Runtime/Engine/Private/UnrealClient.cpp
+++ b/Engine/Source/Runtime/Engine/Private/UnrealClient.cpp
@@ -3,6 +3,7 @@
 
 #include "UnrealClient.h"
 #include "BatchedElements.h"
+#include "ContentStreaming.h"
 #include "Engine/GameViewportClient.h"
 #include "Misc/FileHelper.h"
 #include "DummyViewport.h"
@@ -28,6 +29,7 @@
 #include "Performance/EnginePerformanceTargets.h"
 #include "Elements/Framework/TypedElementList.h"
 #include "EngineUtils.h"
+#include "SceneView.h"
 #include "RenderCounters.h"
 #include "RenderGraphUtils.h"
 #include "DynamicResolutionState.h"
@@ -1770,6 +1772,24 @@ FRHIGPUMask FViewport::GetGPUMask(FRHICommandListImmediate& RHICmdList) const
 	return FRHIGPUMask::GPU0();//RHIGetViewportNextPresentGPUIndex(GetViewportRHI()));
 }
 
+void AddStreamingViewInfo(UWorld& World, const FSceneView& View)
+{
+	const float StreamingScale = 1.f / FMath::Clamp(View.LODDistanceFactor, 0.2f, 1.f);
+	IStreamingManager::Get().AddViewInformation(
+		View.ViewMatrices.GetViewOrigin(),
+		View.UnscaledViewRect.Width(),
+		View.UnscaledViewRect.Width() * View.ViewMatrices.GetViewToClip().M[0][0],
+		StreamingScale, false, 0, nullptr, &World);
+
+	World.ViewLocationsRenderedLastFrame.Add(View.ViewMatrices.GetViewOrigin());
+
+	FWorldCachedViewInfo& Info = World.CachedViewInfoRenderedLastFrame.AddDefaulted_GetRef();
+	Info.ViewMatrix = View.ViewMatrices.GetWorldToView();
+	Info.ProjectionMatrix = View.ViewMatrices.GetViewToClip();
+	Info.ViewProjectionMatrix = View.ViewMatrices.GetWorldToClip();
+	Info.ViewToWorld = View.ViewMatrices.GetViewToWorld();
+}
+
 void FViewport::Draw( bool bShouldPresent /*= true */)
 {
 	SCOPED_NAMED_EVENT(FViewport_Draw, FColor::Red);
@@ -1783,15 +1803,7 @@ void FViewport::Draw( bool bShouldPresent /*= true */)
 	bool bBufferVisualizationDumpingRequired = bAnyScreenshotsRequired && CVarDumpFrames && CVarDumpFrames->GetValueOnGameThread();
 
 	// if this is a game viewport, and game rendering is disabled, then we don't want to actually draw anything
-	if ( World && World->IsGameWorld() && !bIsGameRenderingEnabled)
-	{
-		if (!GUseUnifiedTimeBudgetForStreaming)
-		{
-			// since we aren't drawing the viewport, we still need to update streaming
-			World->UpdateLevelStreaming(); 
-		}
-	}
-	else
+	if (!(World && World->IsGameWorld() && !bIsGameRenderingEnabled))
 	{
 		if( GIsHighResScreenshot )
 		{
@@ -1813,12 +1825,10 @@ void FViewport::Draw( bool bShouldPresent /*= true */)
 		{
 			static const auto CVar = IConsoleManager::Get().FindTConsoleVariableDataInt(TEXT("r.VSync"));
 			bool bLockToVsync = CVar->GetValueOnGameThread() != 0;
-			bool bCameraCut = false;
 			ULocalPlayer* Player = (GEngine && World) ? GEngine->GetFirstGamePlayer(World) : NULL;
 			if (Player && Player->PlayerController)
 			{
 				bLockToVsync |= (Player->PlayerController->bCinematicMode);
-				bCameraCut |= (Player->PlayerController->PlayerCameraManager && Player->PlayerController->PlayerCameraManager->bGameCameraCutThisFrame);
 			}
 	 		EnqueueBeginRenderFrame(bShouldPresent);
 
@@ -1867,27 +1877,8 @@ void FViewport::Draw( bool bShouldPresent /*= true */)
 	
 			// Slate doesn't present immediately. Tag the viewport as requiring vsync so that it happens.
 			SetRequiresVsync(bLockToVsync);
-			SetCameraCut(bCameraCut);
 			EnqueueEndRenderFrame(bLockToVsync, bShouldPresent);
 		}
-		else if (World && World->IsGameWorld() && !GUseUnifiedTimeBudgetForStreaming)
-		{
-			// since we aren't drawing the viewport, we still need to update streaming
-			World->UpdateLevelStreaming();
-		}
-	}
-
-	// Reset the camera cut flags if we are in a viewport that has a world
-	if (World)
-	{
-		for( FConstPlayerControllerIterator Iterator = World->GetPlayerControllerIterator(); Iterator; ++Iterator )
-		{
-			APlayerController* PlayerController = Iterator->Get();
-			if (PlayerController && PlayerController->PlayerCameraManager)
-			{
-				PlayerController->PlayerCameraManager->bGameCameraCutThisFrame = false;
-			}
-		}
 	}
 
 	// countdown the present delay, and then stop the movie at the end
diff --git a/Engine/Source/Runtime/Engine/Private/UnrealEngine.cpp b/Engine/Source/Runtime/Engine/Private/UnrealEngine.cpp
--- a/Engine/Source/Runtime/Engine/Private/UnrealEngine.cpp
+++ b/Engine/Source/Runtime/Engine/Private/UnrealEngine.cpp
@@ -83,6 +83,7 @@ UnrealEngine.cpp: Implements the UEngine class and helpers.
 #include "SystemSettings.h"
 #include "EngineUtils.h"
 #include "Framework/Application/SlateApplication.h"
+#include "Rendering/SlateRenderer.h"
 #include "UObject/GarbageCollectionHistory.h"
 #include "UObject/ReachabilityAnalysis.h"
 #include "Widgets/Input/SButton.h"
@@ -12313,6 +12314,62 @@ void UEngine::HandleUnifiedStreaming(float DeltaSeconds)
 	}
 }
 
+void UEngine::RegisterViewport(TSharedRef<FViewport> Viewport)
+{
+	using namespace UE;
+
+	TUniqueLock Lock(ActiveViewportsMutex);
+	for (const TSharedRef<FViewport>& Existing : ActiveViewports)
+	{
+		if (&Existing.Get() == &Viewport.Get())
+		{
+			return;
+		}
+	}
+	ActiveViewports.Add(Viewport);
+}
+
+void UEngine::UnregisterViewport(TSharedRef<FViewport> Viewport)
+{
+	using namespace UE;
+
+	TUniqueLock Lock(ActiveViewportsMutex);
+	int32 ToRemoveIdx = ActiveViewports.IndexOfByPredicate([&Viewport](const TSharedRef<FViewport>& Existing)
+	{
+		return &Existing.Get() == &Viewport.Get();
+	});
+	if (ToRemoveIdx != INDEX_NONE)
+	{
+		ActiveViewports.RemoveAtSwap(ToRemoveIdx);
+	}
+}
+
+void UEngine::RedrawViewports(bool bShouldPresent)
+{
+	using namespace UE;
+
+	TArray<TSharedRef<FViewport>> ViewportsCopy;
+	{
+		TUniqueLock Lock(ActiveViewportsMutex);
+		ViewportsCopy = ActiveViewports;
+	}
+
+	if (ViewportsCopy.Num() == 0)
+	{
+		return;
+	}
+
+	// Make sure we don't attempt to render a frame in parallel with the loading screen, which might deadlock because of interleaving two different
+	// frames.  It could happen because the Game thread invokes RedrawViewports right after blocking loads when the loading thread is still spinning.
+	FSlateRenderer* SlateRenderer = FSlateApplication::Get().GetRenderer();
+	FScopeLock ScopeLock(SlateRenderer->GetResourceCriticalSection());
+
+	for (const TSharedRef<FViewport>& Viewport : ViewportsCopy)
+	{
+		Viewport->Draw(bShouldPresent);
+	}
+}
+
 void UEngine::EnableScreenSaver( bool bEnable )
 {
 #if PLATFORM_DESKTOP
diff --git a/Engine/Source/Runtime/Engine/Public/EngineStats.h b/Engine/Source/Runtime/Engine/Public/EngineStats.h
--- a/Engine/Source/Runtime/Engine/Public/EngineStats.h
+++ b/Engine/Source/Runtime/Engine/Public/EngineStats.h
@@ -22,7 +22,7 @@ DECLARE_DWORD_COUNTER_STAT_EXTERN(TEXT("Skel Verts CPU Skin"),STAT_CPUSkinVertic
 DECLARE_DWORD_COUNTER_STAT_EXTERN(TEXT("Skel Verts GPU Skin"),STAT_GPUSkinVertices,STATGROUP_Engine, );
 DECLARE_CYCLE_STAT_EXTERN(TEXT("GameEngine Tick"),STAT_GameEngineTick,STATGROUP_Engine, );
 DECLARE_CYCLE_STAT_EXTERN(TEXT("GameViewport Tick"),STAT_GameViewportTick,STATGROUP_Engine, );
-DECLARE_CYCLE_STAT_EXTERN(TEXT("RedrawViewports"),STAT_RedrawViewports,STATGROUP_Engine, );
+DECLARE_CYCLE_STAT_EXTERN(TEXT("RedrawViewports"),STAT_RedrawViewports,STATGROUP_Engine, ENGINE_API);
 DECLARE_CYCLE_STAT_EXTERN(TEXT("Update Level Streaming"),STAT_UpdateLevelStreaming,STATGROUP_Engine, ENGINE_API);
 DECLARE_CYCLE_STAT_EXTERN(TEXT("RHI Game Tick"),STAT_RHITickTime,STATGROUP_Engine, ENGINE_API);
 DECLARE_CYCLE_STAT_EXTERN(TEXT("Debug Hitch"),STAT_IntentionalHitch,STATGROUP_Engine, );
diff --git a/Engine/Source/Runtime/Engine/Public/UnrealClient.h b/Engine/Source/Runtime/Engine/Public/UnrealClient.h
--- a/Engine/Source/Runtime/Engine/Public/UnrealClient.h
+++ b/Engine/Source/Runtime/Engine/Public/UnrealClient.h
@@ -22,9 +22,11 @@
 
 class FCanvas;
 class FRDGBuilder;
+class FSceneView;
 class FViewport;
 class FViewportClient;
 class UModel;
+class UWorld;
 
 /**
  * A render target.
@@ -938,5 +940,10 @@ extern ENGINE_API bool GetViewportScreenShot(FViewport* Viewport, TArray<FColor>
 extern ENGINE_API bool GetViewportScreenShotHDR(FViewport* Viewport, TArray<FLinearColor>& Bitmap, const FIntRect& ViewRect = FIntRect(), const FReadSurfaceDataFlags& ReadSurfaceFlags = FReadSurfaceDataFlags(RCM_MinMax));
 extern ENGINE_API bool GetHighResScreenShotInput(const TCHAR* Cmd, FOutputDevice& Ar, uint32& OutXRes, uint32& OutYRes, float& OutResMult, FIntRect& OutCaptureRegion, bool& OutShouldEnableMask, bool& OutDumpBufferVisualizationTargets, bool& OutCaptureHDR, FString& OutFilenameOverride, bool& OutUseDateTimeAsFileName);
 
+/** Informs a world (and the streaming manager) that a view has been rendered within it during the current frame.
+ * Must be called for asset streaming, and any other features that need to know "where is the scene being drawn from?"
+ */
+ENGINE_API void AddStreamingViewInfo(UWorld& World, const FSceneView& View);
+
 /** Tracks the viewport client that should process the stat command, can be NULL */
 extern ENGINE_API class FCommonViewportClient* GStatProcessingViewportClient;
