Modular Viewports: Fix linker error

#rnx
#Synced-CL 52720326

diff --git a/Engine/Plugins/Runtime/ModularViewports/Source/ModularViewports/Public/Camera/CameraViewportClient.h b/Engine/Plugins/Runtime/ModularViewports/Source/ModularViewports/Public/Camera/CameraViewportClient.h
--- a/Engine/Plugins/Runtime/ModularViewports/Source/ModularViewports/Public/Camera/CameraViewportClient.h
+++ b/Engine/Plugins/Runtime/ModularViewports/Source/ModularViewports/Public/Camera/CameraViewportClient.h
@@ -37,18 +37,18 @@ namespace UE
 		UE_API void CalculateProjectionData(const FViewport& Viewport, FSceneViewProjectionData& OutProjectionData) const;
 
 		// FGCObject interface
-		virtual void AddReferencedObjects(FReferenceCollector& Collector) override;
-		virtual FString GetReferencerName() const override;
+		UE_API virtual void AddReferencedObjects(FReferenceCollector& Collector) override;
+		UE_API virtual FString GetReferencerName() const override;
 		// ~FGCObject interface
 
 		// FViewportClient interface
-		virtual UWorld* GetWorld() const override;
-		virtual void Draw(FViewport* Viewport, FCanvas* Canvas) override;
-		virtual bool InputChar(FViewport* const Viewport, const int32 ControllerId, const TCHAR Character) override;
-		virtual bool IsFocused(FViewport* Viewport) override;
-		virtual bool RequiresHitProxyStorage() override;
-		virtual bool CaptureMouseOnLaunch() override;
-		virtual bool ShouldDPIScaleSceneCanvas() const override;
+		UE_API virtual UWorld* GetWorld() const override;
+		UE_API virtual void Draw(FViewport* Viewport, FCanvas* Canvas) override;
+		UE_API virtual bool InputChar(FViewport* const Viewport, const int32 ControllerId, const TCHAR Character) override;
+		UE_API virtual bool IsFocused(FViewport* Viewport) override;
+		UE_API virtual bool RequiresHitProxyStorage() override;
+		UE_API virtual bool CaptureMouseOnLaunch() override;
+		UE_API virtual bool ShouldDPIScaleSceneCanvas() const override;
 		// ~FViewportClient interface
 	};
 }
